Skeletal Mesh merging plugin and addition of skeleton merging functionality:
+ Added SkeletalMerging plugin - which exposes (blueprint-accessible) API for both SkeletalMesh merging and Skeleton merging
+ Added API to add compatible Skeletons at runtime
+ Add ability to construct FReferenceSkeletonModifier with only a USkeleton
   * required friend-ing to retrieve mutable ReferenceSkeleton
* Engine skeletal merging API input parameters have been exposed to Blueprints
* Switched to Skeleton compatiblity check for Linked animation instances rather than direct skeleton comparison
* Exposed USkeleton as a BlueprintType
* Ensured that runtime generated USkeletons are included in USkeleton::LoadedSkeletons
- Removed EngineTest and QAGame SkeletalMeshMerging code (now part of SkeletalMerging plugin)
   - Added redirectors to patch-up existing test content
   - Enabled new plugin for both EngineTest and QAGame

diff --git a/Engine/Plugins/Runtime/SkeletalMerging/SkeletalMerging.uplugin b/Engine/Plugins/Runtime/SkeletalMerging/SkeletalMerging.uplugin
new file mode 100644
--- /dev/null
+++ b/Engine/Plugins/Runtime/SkeletalMerging/SkeletalMerging.uplugin
@@ -0,0 +1,25 @@
+{
+	"FileVersion" : 1,
+	"Version" : 1,
+	"VersionName" : "1.0",
+	"FriendlyName" : "Skeletal Merging",
+	"Description" : "Provides Blueprint functionality to perform runtime Skeletal Mesh merging",
+	"Category" : "Other",
+	"CreatedBy": "Epic Games, Inc.",
+	"CreatedByURL": "http://epicgames.com",
+	"DocsURL" : "",
+	"MarketplaceURL" : "",
+	"SupportURL" : "",
+	"EnabledByDefault" : false,
+	"CanContainContent" : false,
+	"IsBetaVersion" : false,
+	"Installed" : false,
+	"Modules" :
+	[
+		{
+			"Name" : "SkeletalMerging",
+			"Type" : "Runtime",
+			"LoadingPhase" : "Default"
+		}
+	]
+}
\ No newline at end of file
diff --git a/Engine/Plugins/Runtime/SkeletalMerging/Source/SkeletalMerging/Private/SkeletalMergingLibrary.cpp b/Engine/Plugins/Runtime/SkeletalMerging/Source/SkeletalMerging/Private/SkeletalMergingLibrary.cpp
new file mode 100644
--- /dev/null
+++ b/Engine/Plugins/Runtime/SkeletalMerging/Source/SkeletalMerging/Private/SkeletalMergingLibrary.cpp
@@ -0,0 +1,436 @@
+// Copyright Epic Games, Inc. All Rights Reserved.
+
+#include "SkeletalMergingLibrary.h"
+#include "SkeletalMeshMerge.h"
+#include "Engine/SkeletalMeshSocket.h"
+#include "Engine/SkeletalMesh.h"
+#include "Animation/Skeleton.h"
+#include "Algo/Accumulate.h"
+
+IMPLEMENT_MODULE(IModuleInterface, SkeletalMerging);
+
+namespace UE
+{
+	namespace SkeletonMerging
+	{
+		// Helper structure to merge bone hierarchies together and populate a FReferenceSkeleton with the result(s)
+		struct FMergedBoneHierarchy
+		{
+			FMergedBoneHierarchy(uint32 NumExpectedBones)
+			{
+				BoneNamePose.Reserve(NumExpectedBones);
+				PathToBoneNames.Reserve(NumExpectedBones);
+				PathHashToBoneNames.Reserve(NumExpectedBones);
+			}
+
+			void AddBone(const FName& BoneName, const FTransform& ReferencePose, uint32 PathHash)
+			{
+				// Store reference transform according to bone name hash
+				BoneNamePose.Add(BoneName, ReferencePose);
+
+				// Add bone as child to parent path
+				PathHashToBoneNames.FindOrAdd(PathHash).Add(BoneName);
+
+				// Append bone hash to parent path and store 
+				const uint32 BoneHash = GetTypeHash(BoneName);
+				PathToBoneNames.Add(BoneName, HashCombine(PathHash, BoneHash));
+			}
+
+			void PopulateSkeleton(FReferenceSkeletonModifier& SkeletonModifier)
+			{
+				const uint32 Zero = 0;
+				const uint32 RootParentHash = HashCombine(Zero, Zero);
+
+				// Root bone is always parented to 0 hash data entry, so we expect a single root-bone (child)
+				const TSet<FName>& ChildBoneNames = GetChildBonesForPath(RootParentHash);
+				ensure(ChildBoneNames.Num() == 1);
+				const FName RootBoneName = ChildBoneNames.Array()[0];
+
+				// Add root-bone and traverse data to populate child hierarchies
+				const FMeshBoneInfo BoneInfo(RootBoneName, RootBoneName.ToString(), INDEX_NONE);
+				SkeletonModifier.Add(BoneInfo, GetReferencePose(RootBoneName));
+
+				RecursiveAddBones(SkeletonModifier, RootBoneName);
+			}
+		protected:			
+			const FTransform& GetReferencePose(const FName& InName) const { return BoneNamePose.FindChecked(InName); }
+			uint32 GetBonePathHash(const FName& InName) const { return PathToBoneNames.FindChecked(InName); }
+			const TSet<FName>* FindChildBonesForPath(uint32 InPath) const { return PathHashToBoneNames.Find(InPath); }
+			const TSet<FName>& GetChildBonesForPath(uint32 InPath) const { return PathHashToBoneNames.FindChecked(InPath); }
+			
+			void RecursiveAddBones(FReferenceSkeletonModifier& SkeletonModifier, const FName ParentBoneName)
+			{
+				const uint32 PathHash = GetBonePathHash(ParentBoneName);
+				if (const TSet<FName>* BoneNamesPtr = FindChildBonesForPath(PathHash))
+				{
+					for (const FName& ChildBoneName : *BoneNamesPtr)
+					{
+						FMeshBoneInfo BoneInfo(ChildBoneName, ChildBoneName.ToString(), SkeletonModifier.FindBoneIndex(ParentBoneName));
+						SkeletonModifier.Add(BoneInfo, GetReferencePose(ChildBoneName));
+						RecursiveAddBones(SkeletonModifier, ChildBoneName);
+					}
+				}
+			}
+
+		private:			
+			// Reference pose transform for given bone name
+			TMap<FName, FTransform> BoneNamePose;
+			// Accumulated hierarchy hash from bone to root bone
+			TMap<FName, uint32> PathToBoneNames;
+			// Set of child bones for given hierarchy hash
+			TMap<uint32, TSet<FName>> PathHashToBoneNames;
+		};
+	}
+}
+
+USkeleton* USkeletalMergingLibrary::MergeSkeletons(const FSkeletonMergeParams& Params)
+{
+	// List of unique skeletons generated from input parameters
+	TArray<TObjectPtr<USkeleton>> ToMergeSkeletons;
+	for ( TObjectPtr<USkeleton> SkeletonPtr : Params.SkeletonsToMerge)
+	{
+		ToMergeSkeletons.AddUnique(SkeletonPtr);
+	}
+
+	// Ensure we have at least one valid Skeleton to merge
+	const int32 NumberOfSkeletons = ToMergeSkeletons.Num();
+	if (NumberOfSkeletons == 0)
+	{
+		return nullptr;
+	}
+
+	// Calculate potential total number of bones, used for pre-allocating data arrays
+	const int32 TotalPossibleBones = Algo::Accumulate(ToMergeSkeletons, 0, [](int32 Sum, const USkeleton* Skeleton)
+	{
+		return Sum + Skeleton->GetReferenceSkeleton().GetRawBoneNum();
+	});
+
+	// Ensure a valid skeleton (number of bones) will be generated
+	if (TotalPossibleBones == 0)
+	{
+		return nullptr;
+	}
+
+	UE::SkeletonMerging::FMergedBoneHierarchy MergedBoneHierarchy(TotalPossibleBones);
+	
+	// Accumulated hierarchy hash from parent-bone to root bone
+	TMap<FName, uint32> BoneNamesToPathHash;
+	BoneNamesToPathHash.Reserve(TotalPossibleBones);
+
+	// Combined bone and socket name hash
+	TMap<uint32, TObjectPtr<USkeletalMeshSocket>> HashToSockets;
+	// Combined from and to-bone name hash
+	TMap<uint32, const FVirtualBone*> HashToVirtualBones;
+
+	TMap<FName, const FCurveMetaData*> UniqueCurveNames;
+	TMap<FName, TSet<FName>> GroupToSlotNames;
+	TMap<FName, TArray<const UBlendProfile*>> UniqueBlendProfiles;
+
+	for (int32 SkeletonIndex = 0; SkeletonIndex < NumberOfSkeletons; ++SkeletonIndex)
+	{
+		const USkeleton* Skeleton = ToMergeSkeletons[SkeletonIndex];
+		const FReferenceSkeleton& ReferenceSkeleton = Skeleton->GetReferenceSkeleton();
+		const TArray<FMeshBoneInfo>& Bones = ReferenceSkeleton.GetRawRefBoneInfo();		
+		const TArray<FTransform>& BonePoses = ReferenceSkeleton.GetRawRefBonePose();
+
+		const int32 NumBones = Bones.Num();
+		for (int32 BoneIndex = 0; BoneIndex < NumBones; ++BoneIndex)
+		{
+			const FMeshBoneInfo& Bone = Bones[BoneIndex];
+
+			// Retrieve parent bone name and respective hash, root-bone is assumed to have a parent hash of 0
+			const FName ParentName = Bone.ParentIndex != INDEX_NONE ? Bones[Bone.ParentIndex].Name : NAME_None;
+			const uint32 ParentHash = Bone.ParentIndex != INDEX_NONE ? GetTypeHash(ParentName) : 0;
+
+			// Look-up the path-hash from root to the parent bone
+			const uint32* ParentPath = BoneNamesToPathHash.Find(ParentName);
+			const uint32 ParentPathHash = ParentPath ? *ParentPath : 0;
+
+			// Append parent hash to path to give full path hash to current bone
+			const uint32 BonePathHash = HashCombine(ParentPathHash, ParentHash);
+			BoneNamesToPathHash.Add(Bone.Name, BonePathHash);
+
+			// Add bone to hierarchy
+			MergedBoneHierarchy.AddBone(Bone.Name, BonePoses[BoneIndex], BonePathHash);
+		}
+
+		if (Params.bMergeSockets)
+		{
+			for (const TObjectPtr<USkeletalMeshSocket>& Socket : Skeleton->Sockets)
+			{
+				const uint32 Hash = HashCombine(GetTypeHash(Socket->SocketName), GetTypeHash(Socket->BoneName));
+				HashToSockets.Add(Hash, Socket);
+			}
+		}
+
+		if (Params.bMergeVirtualBones)
+		{
+			const TArray<FVirtualBone>& VirtualBones = Skeleton->GetVirtualBones();		
+	        for (const FVirtualBone& VB : VirtualBones)
+	        {
+	            const uint32 Hash = HashCombine(GetTypeHash(VB.SourceBoneName), GetTypeHash(VB.TargetBoneName));
+	            HashToVirtualBones.Add(Hash, &VB);
+	        }
+		}		
+
+		if (Params.bMergeCurveNames)
+		{
+			if (const FSmartNameMapping* CurveMappingPtr = Skeleton->GetSmartNameContainer(USkeleton::AnimCurveMappingName))
+			{
+				TArray<FName> CurveNames;
+				CurveMappingPtr->FillNameArray(CurveNames);
+				for (const FName& CurveName : CurveNames)
+				{
+					UniqueCurveNames.FindOrAdd(CurveName) = CurveMappingPtr->GetCurveMetaData(CurveName);
+				}
+			}
+		}
+
+		if (Params.bMergeAnimSlotGroups)
+		{
+			const TArray<FAnimSlotGroup>& SlotGroups = Skeleton->GetSlotGroups();
+			for (const FAnimSlotGroup& AnimSlotGroup : SlotGroups)
+			{
+				GroupToSlotNames.FindOrAdd(AnimSlotGroup.GroupName).Append(AnimSlotGroup.SlotNames);
+			}
+		}
+		
+		if (Params.bMergeBlendProfiles)
+		{
+			for ( const TObjectPtr<UBlendProfile>& BlendProfile : Skeleton->BlendProfiles)
+			{
+				UniqueBlendProfiles.FindOrAdd(BlendProfile->GetFName()).Add(BlendProfile.Get());
+			}			
+		}		
+	}
+
+	USkeleton* GeneratedSkeleton = NewObject<USkeleton>();
+
+	// Generate bone hierarchy
+	{		
+		FReferenceSkeletonModifier Modifier(GeneratedSkeleton);
+		MergedBoneHierarchy.PopulateSkeleton(Modifier);
+	}
+	
+	// Merge sockets
+	if (Params.bMergeSockets)
+	{
+		TArray<TObjectPtr<USkeletalMeshSocket>> Sockets;
+		HashToSockets.GenerateValueArray(Sockets);
+		AddSockets(GeneratedSkeleton, Sockets);
+	}
+	
+	// Merge virtual bones
+	if (Params.bMergeVirtualBones)
+	{
+		TArray<const FVirtualBone*> VirtualBones;
+		HashToVirtualBones.GenerateValueArray(VirtualBones);
+		AddVirtualBones(GeneratedSkeleton,VirtualBones);
+	}
+
+	// Merge Curve / track mappings	
+	if (Params.bMergeCurveNames)
+	{
+		AddCurveNames(GeneratedSkeleton, UniqueCurveNames);
+	}
+
+	// Merge blend profiles
+	if (Params.bMergeBlendProfiles)
+	{
+		AddBlendProfiles(GeneratedSkeleton, UniqueBlendProfiles);
+	}
+
+	// Merge SlotGroups
+	if (Params.bMergeAnimSlotGroups)
+	{
+		AddAnimationSlotGroups(GeneratedSkeleton, GroupToSlotNames);
+	}
+		
+	return GeneratedSkeleton;
+}
+
+void USkeletalMergingLibrary::AddSockets(USkeleton* InSkeleton, const TArray<TObjectPtr<USkeletalMeshSocket>>& InSockets)
+{
+	for (const TObjectPtr<USkeletalMeshSocket>& MergeSocket : InSockets)
+	{
+		USkeletalMeshSocket* NewSocket = NewObject<USkeletalMeshSocket>(InSkeleton);
+		if (NewSocket != nullptr)
+		{
+			InSkeleton->Sockets.Add(NewSocket);
+
+			// Copy over all socket information
+			NewSocket->SocketName = MergeSocket->SocketName;
+			NewSocket->BoneName = MergeSocket->BoneName;
+			NewSocket->RelativeLocation = MergeSocket->RelativeLocation;
+			NewSocket->RelativeRotation = MergeSocket->RelativeRotation;
+			NewSocket->RelativeScale = MergeSocket->RelativeScale;
+			NewSocket->bForceAlwaysAnimated = MergeSocket->bForceAlwaysAnimated;
+		}
+	}
+}
+
+void USkeletalMergingLibrary::AddVirtualBones(USkeleton* InSkeleton, const TArray<const FVirtualBone*> InVirtualBones)
+{
+	for(const FVirtualBone* VirtualBone : InVirtualBones)
+	{
+		FName VirtualBoneName = NAME_None;				
+		InSkeleton->AddNewVirtualBone(VirtualBone->SourceBoneName, VirtualBone->TargetBoneName, VirtualBoneName);
+	}	
+}
+
+void USkeletalMergingLibrary::AddCurveNames(USkeleton* InSkeleton, const TMap<FName, const FCurveMetaData*>& InCurves)
+{
+	TArray<FSmartName> CurveSmartNames;
+
+	Algo::Transform(InCurves, CurveSmartNames, [](const TPair<FName, const FCurveMetaData*>& CurveMetaDataPair)
+	{
+		return FSmartName(CurveMetaDataPair.Key, INDEX_NONE);
+	});
+	InSkeleton->VerifySmartNames(USkeleton::AnimCurveMappingName, CurveSmartNames);
+		
+	for(const TPair<FName, const FCurveMetaData*> CurveMetaDataPair : InCurves)
+	{
+		if (CurveMetaDataPair.Value)
+		{
+			*InSkeleton->GetCurveMetaData(CurveMetaDataPair.Key) = *CurveMetaDataPair.Value;
+		}
+	}
+}
+
+void USkeletalMergingLibrary::AddBlendProfiles(USkeleton* InSkeleton, const TMap<FName, TArray<const UBlendProfile*>>& InBlendProfiles)
+{
+	for (const TPair<FName, TArray<const UBlendProfile*>>& BlendProfilesPair : InBlendProfiles)
+	{
+		const TArray<const UBlendProfile*>& BlendProfiles = BlendProfilesPair.Value;
+		UBlendProfile* MergedBlendProfile = InSkeleton->CreateNewBlendProfile(BlendProfilesPair.Key);
+			
+		for (int32 ProfileIndex = 0; ProfileIndex < BlendProfiles.Num(); ++ProfileIndex)
+		{
+			const UBlendProfile* Profile = BlendProfiles[ProfileIndex];				
+			MergedBlendProfile->Mode = ProfileIndex == 0 ? Profile->Mode : MergedBlendProfile->Mode;
+
+			// Mismatch in terms of blend profile type
+			ensure(MergedBlendProfile->Mode == Profile->Mode);
+
+			for (const FBlendProfileBoneEntry& Entry : Profile->ProfileEntries)
+			{
+				// Overlapping bone entries
+				ensure(!MergedBlendProfile->ProfileEntries.ContainsByPredicate([Entry](const FBlendProfileBoneEntry& InEntry)
+				{
+					return InEntry.BoneReference.BoneName == Entry.BoneReference.BoneName;
+				}));
+
+				MergedBlendProfile->SetBoneBlendScale(Entry.BoneReference.BoneName, Entry.BlendScale, false, true);
+			}
+		}
+	}
+}
+
+void USkeletalMergingLibrary::AddAnimationSlotGroups(USkeleton* InSkeleton, const TMap<FName, TSet<FName>>& InSlotGroupsNames)
+{
+	for (const TPair<FName, TSet<FName>>& SlotGroupNamePair : InSlotGroupsNames)
+	{
+		InSkeleton->AddSlotGroupName(SlotGroupNamePair.Key);
+		for (const FName& SlotName : SlotGroupNamePair.Value)
+		{
+			InSkeleton->SetSlotGroupName(SlotName, SlotGroupNamePair.Key);				
+		}
+	}
+}
+
+USkeletalMesh* USkeletalMergingLibrary::MergeMeshes(const FSkeletalMeshMergeParams& Params)
+{
+	TArray<USkeletalMesh*> MeshesToMergeCopy = Params.MeshesToMerge;
+	MeshesToMergeCopy.RemoveAll([](USkeletalMesh* InMesh)
+	{
+		return InMesh == nullptr;
+	});
+
+	if (MeshesToMergeCopy.Num() <= 1)
+	{
+		UE_LOG(LogTemp, Warning, TEXT("Must provide multiple valid Skeletal Meshes in order to perform a merge."));
+		return nullptr;
+	}
+
+	EMeshBufferAccess BufferAccess = Params.bNeedsCpuAccess ?
+										EMeshBufferAccess::ForceCPUAndGPU :
+										EMeshBufferAccess::Default;
+	
+	bool bRunDuplicateCheck = false;
+	USkeletalMesh* BaseMesh = NewObject<USkeletalMesh>();
+	
+	if (Params.Skeleton && Params.bSkeletonBefore)
+	{
+		BaseMesh->SetSkeleton(Params.Skeleton);
+		bRunDuplicateCheck = true;
+
+		for (USkeletalMeshSocket* Socket : BaseMesh->GetMeshOnlySocketList())
+		{
+			if (Socket)
+			{
+				UE_LOG(LogTemp, Warning, TEXT("SkelMeshSocket: %s"), *(Socket->SocketName.ToString()));
+			}
+		}
+
+		for (USkeletalMeshSocket* Socket : BaseMesh->GetSkeleton()->Sockets)
+		{
+			if (Socket)
+			{
+				UE_LOG(LogTemp, Warning, TEXT("SkelSocket: %s"), *(Socket->SocketName.ToString()));
+			}
+		}
+	}
+
+	FSkelMeshMergeUVTransformMapping Mapping;
+	Mapping.UVTransformsPerMesh = Params.UVTransformsPerMesh;
+	FSkeletalMeshMerge Merger(BaseMesh, MeshesToMergeCopy, Params.MeshSectionMappings, Params.StripTopLODS, BufferAccess, &Mapping);
+	if (!Merger.DoMerge())
+	{
+		UE_LOG(LogTemp, Warning, TEXT("Merge failed!"));
+		return nullptr;
+	}
+
+	if (Params.Skeleton && !Params.bSkeletonBefore)
+	{
+		BaseMesh->SetSkeleton(Params.Skeleton);
+	}
+
+	if (bRunDuplicateCheck)
+	{
+		TArray<FName> SkelMeshSockets;
+		TArray<FName> SkelSockets;
+
+		for (USkeletalMeshSocket* Socket : BaseMesh->GetMeshOnlySocketList())
+		{
+			if (Socket)
+			{
+				SkelMeshSockets.Add(Socket->GetFName());
+				UE_LOG(LogTemp, Warning, TEXT("SkelMeshSocket: %s"), *(Socket->SocketName.ToString()));
+			}
+		}
+
+		for (USkeletalMeshSocket* Socket : BaseMesh->GetSkeleton()->Sockets)
+		{
+			if (Socket)
+			{
+				SkelSockets.Add(Socket->GetFName());
+				UE_LOG(LogTemp, Warning, TEXT("SkelSocket: %s"), *(Socket->SocketName.ToString()));
+			}
+		}
+
+		TSet<FName> UniqueSkelMeshSockets;
+		TSet<FName> UniqueSkelSockets;
+
+		UniqueSkelMeshSockets.Append(SkelMeshSockets);
+		UniqueSkelSockets.Append(SkelSockets);
+
+		int32 Total = SkelSockets.Num() + SkelMeshSockets.Num();
+		int32 UniqueTotal = UniqueSkelMeshSockets.Num() + UniqueSkelSockets.Num();
+
+		UE_LOG(LogTemp, Warning, TEXT("SkelMeshSocketCount: %d | SkelSocketCount: %d | Combined: %d"), SkelMeshSockets.Num(), SkelSockets.Num(), Total);
+		UE_LOG(LogTemp, Warning, TEXT("SkelMeshSocketCount: %d | SkelSocketCount: %d | Combined: %d"), UniqueSkelMeshSockets.Num(), UniqueSkelSockets.Num(), UniqueTotal);
+		UE_LOG(LogTemp, Warning, TEXT("Found Duplicates: %s"), *((Total != UniqueTotal) ? FString("True") : FString("False")));
+	}
+
+	return BaseMesh;
+}
diff --git a/Engine/Plugins/Runtime/SkeletalMerging/Source/SkeletalMerging/Public/SkeletalMergingLibrary.h b/Engine/Plugins/Runtime/SkeletalMerging/Source/SkeletalMerging/Public/SkeletalMergingLibrary.h
new file mode 100644
--- /dev/null
+++ b/Engine/Plugins/Runtime/SkeletalMerging/Source/SkeletalMerging/Public/SkeletalMergingLibrary.h
@@ -0,0 +1,119 @@
+// Copyright Epic Games, Inc. All Rights Reserved.
+
+#pragma once
+
+#include "SkeletalMeshMerge.h"
+#include "SkeletalMergingLibrary.generated.h"
+
+/**
+* Struct containing all parameters used to perform a Skeletal Mesh merge.
+*/
+USTRUCT(BlueprintType)
+struct SKELETALMERGING_API FSkeletalMeshMergeParams
+{
+	GENERATED_USTRUCT_BODY()
+
+	FSkeletalMeshMergeParams()
+	{
+		MeshSectionMappings = TArray<FSkelMeshMergeSectionMapping>();
+		UVTransformsPerMesh = TArray<FSkelMeshMergeMeshUVTransforms>();
+
+		StripTopLODS = 0;
+		bNeedsCpuAccess = false;
+		bSkeletonBefore = false;
+		Skeleton = nullptr;
+	}
+
+	// An optional array to map sections from the source meshes to merged section entries
+	UPROPERTY(EditAnywhere, BlueprintReadWrite, Category=SkeletalMeshMerge)
+	TArray<FSkelMeshMergeSectionMapping> MeshSectionMappings;
+
+	// An optional array to transform the UVs in each mesh
+	UPROPERTY(EditAnywhere, BlueprintReadWrite, Category=SkeletalMeshMerge)
+	TArray<FSkelMeshMergeMeshUVTransforms> UVTransformsPerMesh;
+
+	// The list of skeletal meshes to merge.
+	UPROPERTY(EditAnywhere, BlueprintReadWrite, Category=SkeletalMeshMerge)
+	TArray<TObjectPtr<USkeletalMesh>> MeshesToMerge;
+
+	// The number of high LODs to remove from input meshes
+	UPROPERTY(EditAnywhere, BlueprintReadWrite, Category=SkeletalMeshMerge)
+	int32 StripTopLODS;
+
+	// Whether or not the resulting mesh needs to be accessed by the CPU for any reason (e.g. for spawning particle effects).
+	UPROPERTY(EditAnywhere, BlueprintReadWrite, Category=SkeletalMeshMerge)
+	uint32 bNeedsCpuAccess:1;
+
+	// Update skeleton before merge. Otherwise, update after.
+	// Skeleton must also be provided.
+	UPROPERTY(EditAnywhere, BlueprintReadWrite, Category=SkeletalMeshMerge)
+	uint32 bSkeletonBefore:1;
+
+	// Skeleton that will be used for the merged mesh.
+	// Leave empty if the generated skeleton is OK.
+	UPROPERTY(EditAnywhere, BlueprintReadWrite, Category=SkeletalMeshMerge)
+	TObjectPtr<USkeleton> Skeleton;
+};
+
+/**
+* Struct containing all parameters used to perform a Skeleton merge.
+*/
+USTRUCT(BlueprintType)
+struct SKELETALMERGING_API FSkeletonMergeParams
+{
+	GENERATED_USTRUCT_BODY()
+
+	// The list of skeletons to merge.
+	UPROPERTY(EditAnywhere, BlueprintReadWrite, Category=SkeletonMerge)
+	TArray<TObjectPtr<USkeleton>> SkeletonsToMerge;
+
+	// Whether or not to include Sockets when merging the Skeletons
+	UPROPERTY(EditAnywhere, BlueprintReadWrite, Category=SkeletonMerge)
+	bool bMergeSockets = true;
+
+	// Whether or not to include Virtual Bones when merging the Skeletons
+	UPROPERTY(EditAnywhere, BlueprintReadWrite, Category=SkeletonMerge)
+	bool bMergeVirtualBones = true;
+
+	// Whether or not to include (Animation) Curve names when merging the Skeletons
+	UPROPERTY(EditAnywhere, BlueprintReadWrite, Category=SkeletonMerge)
+	bool bMergeCurveNames = true;
+
+	// Whether or not to include Blend Profiles when merging the Skeletons
+	UPROPERTY(EditAnywhere, BlueprintReadWrite, Category=SkeletonMerge)
+	bool bMergeBlendProfiles = true;
+
+	// Whether or not to include Animation Slot Group (names) when merging the Skeletons
+	UPROPERTY(EditAnywhere, BlueprintReadWrite, Category=SkeletonMerge)
+	bool bMergeAnimSlotGroups = true;
+};
+
+/**
+* Component that can be used to perform Skeletal Mesh merges from Blueprints.
+*/
+UCLASS( ClassGroup=(Custom) )
+class SKELETALMERGING_API USkeletalMergingLibrary : public UBlueprintFunctionLibrary
+{
+	GENERATED_BODY()
+public:	
+	/**
+	* Merges the given meshes into a single mesh.
+	* @return The merged mesh (will be invalid if the merge failed).
+	*/
+	UFUNCTION(BlueprintCallable, Category="Mesh Merge", meta=(UnsafeDuringActorConstruction="true"))
+	static class USkeletalMesh* MergeMeshes(const FSkeletalMeshMergeParams& Params);
+	
+	/**
+	* Merges the skeletons for the provided meshes into a single skeleton.
+	* @return The merged Skeleton
+	*/
+	UFUNCTION(BlueprintCallable, Category="Skeleton Merge", meta=(UnsafeDuringActorConstruction="true"))
+	static class USkeleton* MergeSkeletons(const FSkeletonMergeParams& Params);
+
+protected:
+	static void AddSockets(USkeleton* InSkeleton, const TArray<TObjectPtr<class USkeletalMeshSocket>>& InSockets);
+	static void AddVirtualBones(USkeleton* InSkeleton, const TArray<const struct FVirtualBone*> InVirtualBones);
+	static void AddCurveNames(USkeleton* InSkeleton, const TMap<FName, const struct FCurveMetaData*>& InCurves);
+	static void AddBlendProfiles(USkeleton* InSkeleton, const TMap<FName, TArray<const class UBlendProfile*>>& InBlendProfiles);
+	static void AddAnimationSlotGroups(USkeleton* InSkeleton, const TMap<FName, TSet<FName>>& InSlotGroupsNames);
+};
diff --git a/Engine/Plugins/Runtime/SkeletalMerging/Source/SkeletalMerging/SkeletalMerging.build.cs b/Engine/Plugins/Runtime/SkeletalMerging/Source/SkeletalMerging/SkeletalMerging.build.cs
new file mode 100644
--- /dev/null
+++ b/Engine/Plugins/Runtime/SkeletalMerging/Source/SkeletalMerging/SkeletalMerging.build.cs
@@ -0,0 +1,18 @@
+// Copyright Epic Games, Inc. All Rights Reserved.
+
+using UnrealBuildTool;
+
+public class SkeletalMerging : ModuleRules
+{
+	public SkeletalMerging(ReadOnlyTargetRules Target) : base(Target)
+	{
+		PrivateDependencyModuleNames.AddRange(
+			new string[]
+			{
+				"Core",
+				"CoreUObject",
+				"Engine"
+			}
+		);
+	}
+}
\ No newline at end of file
diff --git a/Engine/Source/Runtime/Engine/Classes/Animation/Skeleton.h b/Engine/Source/Runtime/Engine/Classes/Animation/Skeleton.h
--- a/Engine/Source/Runtime/Engine/Classes/Animation/Skeleton.h
+++ b/Engine/Source/Runtime/Engine/Classes/Animation/Skeleton.h
@@ -444,7 +444,7 @@ private:
  *		- Bone/track linkup between mesh and animation
  *		- Retargetting related
  */
-UCLASS(hidecategories=Object, MinimalAPI)
+UCLASS(hidecategories=Object, MinimalAPI, BlueprintType)
 class USkeleton : public UObject, public IInterface_AssetUserData, public IInterface_PreviewMeshProvider
 {
 	friend class UAnimationBlueprintLibrary;
@@ -819,6 +819,10 @@ public:
 	DECLARE_EVENT(USkeleton, FSmartNamesChangedEvent);
 	FSmartNamesChangedEvent OnSmartNamesChangedEvent;
 
+public:
+	UFUNCTION(BlueprintCallable, Category=Skeleton)
+	ENGINE_API void AddCompatibleSkeleton(const USkeleton* SourceSkeleton);
+	ENGINE_API void RemoveCompatibleSkeleton(const USkeleton* SourceSkeleton);
 protected:
 	mutable FCriticalSection SkeletonRemappingMutex;
 
@@ -1156,5 +1160,7 @@ protected:
 	/** Array of user data stored with the asset */
 	UPROPERTY(EditAnywhere, AdvancedDisplay, Instanced, Category = Skeleton)
 	TArray<TObjectPtr<UAssetUserData>> AssetUserData;
+
+	friend struct FReferenceSkeletonModifier;
 };
 
diff --git a/Engine/Source/Runtime/Engine/Private/Animation/AnimInstance.cpp b/Engine/Source/Runtime/Engine/Private/Animation/AnimInstance.cpp
--- a/Engine/Source/Runtime/Engine/Private/Animation/AnimInstance.cpp
+++ b/Engine/Source/Runtime/Engine/Private/Animation/AnimInstance.cpp
@@ -2801,7 +2801,7 @@ void UAnimInstance::PerformLinkedLayerOverlayOperation(TSubclassOf<UAnimInstance
 			IAnimClassInterface* LinkedAnimBlueprintClass = IAnimClassInterface::GetFromClass(NewClass);
 			USkeleton* LinkedSkeleton = LinkedAnimBlueprintClass->GetTargetSkeleton();
 			USkeleton* OuterSkeleton = AnimBlueprintClass->GetTargetSkeleton();
-			if(LinkedSkeleton != OuterSkeleton)
+			if(!LinkedSkeleton->IsCompatible(OuterSkeleton))
 			{
 				UE_LOG(LogAnimation, Warning, TEXT("Linking layer: Linked instance class has a mismatched target skeleton. Expected %s, found %s."), OuterSkeleton ? *OuterSkeleton->GetName() : TEXT("null"), LinkedSkeleton ? *LinkedSkeleton->GetName() : TEXT("null"));
 				return;
diff --git a/Engine/Source/Runtime/Engine/Private/Animation/AnimNode_LinkedAnimGraph.cpp b/Engine/Source/Runtime/Engine/Private/Animation/AnimNode_LinkedAnimGraph.cpp
--- a/Engine/Source/Runtime/Engine/Private/Animation/AnimNode_LinkedAnimGraph.cpp
+++ b/Engine/Source/Runtime/Engine/Private/Animation/AnimNode_LinkedAnimGraph.cpp
@@ -309,7 +309,7 @@ void FAnimNode_LinkedAnimGraph::SetAnimClass(TSubclassOf<UAnimInstance> InClass,
 		IAnimClassInterface* OuterAnimBlueprintClass = IAnimClassInterface::GetFromClass(InOwningAnimInstance->GetClass());
 		USkeleton* LinkedSkeleton = LinkedAnimBlueprintClass->GetTargetSkeleton();
 		USkeleton* OuterSkeleton = OuterAnimBlueprintClass->GetTargetSkeleton();
-		if(LinkedSkeleton != OuterSkeleton)
+		if(!LinkedSkeleton->IsCompatible(OuterSkeleton))
 		{
 			UE_LOG(LogAnimation, Warning, TEXT("Setting linked anim instance class: Class has a mismatched target skeleton. Expected %s, found %s."), OuterSkeleton ? *OuterSkeleton->GetName() : TEXT("null"), LinkedSkeleton ? *LinkedSkeleton->GetName() : TEXT("null"));
 			return;
diff --git a/Engine/Source/Runtime/Engine/Private/Animation/ReferenceSkeleton.cpp b/Engine/Source/Runtime/Engine/Private/Animation/ReferenceSkeleton.cpp
--- a/Engine/Source/Runtime/Engine/Private/Animation/ReferenceSkeleton.cpp
+++ b/Engine/Source/Runtime/Engine/Private/Animation/ReferenceSkeleton.cpp
@@ -4,6 +4,8 @@
 #include "Animation/Skeleton.h"
 #include "Engine/SkeletalMesh.h"
 
+FReferenceSkeletonModifier::FReferenceSkeletonModifier(USkeleton* InSkeleton) : RefSkeleton(InSkeleton->ReferenceSkeleton), Skeleton(InSkeleton) {}
+
 FReferenceSkeletonModifier::~FReferenceSkeletonModifier()
 {
 	RefSkeleton.RebuildRefSkeleton(Skeleton, true);
diff --git a/Engine/Source/Runtime/Engine/Private/Animation/Skeleton.cpp b/Engine/Source/Runtime/Engine/Private/Animation/Skeleton.cpp
--- a/Engine/Source/Runtime/Engine/Private/Animation/Skeleton.cpp
+++ b/Engine/Source/Runtime/Engine/Private/Animation/Skeleton.cpp
@@ -158,6 +158,22 @@ bool USkeleton::IsCompatibleSkeletonByAssetString(const FString& SkeletonAssetSt
 	return false;
 }
 
+void USkeleton::AddCompatibleSkeleton(const USkeleton* SourceSkeleton)
+{
+	CompatibleSkeletons.AddUnique(SourceSkeleton);
+
+	ClearSkeletonRemappings();
+	BuildSkeletonRemappings(false);
+}
+
+void USkeleton::RemoveCompatibleSkeleton(const USkeleton* SourceSkeleton)
+{
+	CompatibleSkeletons.Remove(SourceSkeleton);
+
+	ClearSkeletonRemappings();
+	BuildSkeletonRemappings(false);
+}
+
 bool USkeleton::IsCompatibleSkeletonByAssetData(const FAssetData& AssetData, const TCHAR* InTag) const
 {
 	return IsCompatibleSkeletonByAssetString(AssetData.GetTagValueRef<FString>(InTag));
@@ -452,6 +468,14 @@ USkeleton::USkeleton(const FObjectInitializer& ObjectInitializer)
 	{
 		FCoreUObjectDelegates::OnPackageReloaded.AddStatic(&USkeleton::HandlePackageReloaded);
 	}
+	else
+	{
+		if (!HasAnyFlags(RF_NeedPostLoad))
+		{
+			FScopeLock Lock(&LoadedSkeletonsMutex);
+			LoadedSkeletons.Add(this);
+		}
+	}
 }
 
 void USkeleton::BeginDestroy()
@@ -2490,14 +2514,7 @@ void USkeleton::GetAssetRegistryTags(TArray<FAssetRegistryTag>& OutTags) const
 	OutTags.Add(FAssetRegistryTag(USkeleton::RigTag, RigFullName, FAssetRegistryTag::TT_Hidden));
 }
 
-UBlendProfile* USkeleton::CreateNewBlendProfile(const FName& InProfileName)
-{
-	Modify();
-	UBlendProfile* NewProfile = NewObject<UBlendProfile>(this, InProfileName, RF_Public | RF_Transactional);
-	BlendProfiles.Add(NewProfile);
-
-	return NewProfile;
-}
+#endif //WITH_EDITOR
 
 UBlendProfile* USkeleton::GetBlendProfile(const FName& InProfileName)
 {
@@ -2513,7 +2530,14 @@ UBlendProfile* USkeleton::GetBlendProfile(const FName& InProfileName)
 	return nullptr;
 }
 
-#endif //WITH_EDITOR
+UBlendProfile* USkeleton::CreateNewBlendProfile(const FName& InProfileName)
+{
+	Modify();
+	UBlendProfile* NewProfile = NewObject<UBlendProfile>(this, InProfileName, RF_Public | RF_Transactional);
+	BlendProfiles.Add(NewProfile);
+
+	return NewProfile;
+}
 
 USkeletalMeshSocket* USkeleton::FindSocket(FName InSocketName) const
 {
diff --git a/Engine/Source/Runtime/Engine/Private/SkeletalMeshMerge.cpp b/Engine/Source/Runtime/Engine/Private/SkeletalMeshMerge.cpp
--- a/Engine/Source/Runtime/Engine/Private/SkeletalMeshMerge.cpp
+++ b/Engine/Source/Runtime/Engine/Private/SkeletalMeshMerge.cpp
@@ -28,7 +28,7 @@ FSkeletalMeshMerge::FSkeletalMeshMerge(USkeletalMesh* InMergeMesh,
 									   const TArray<FSkelMeshMergeSectionMapping>& InForceSectionMapping,
 									   int32 InStripTopLODs,
                                        EMeshBufferAccess InMeshBufferAccess,
-									   FSkelMeshMergeUVTransforms* InSectionUVTransforms)
+									   const FSkelMeshMergeUVTransformMapping* InSectionUVTransforms)
 :	MergeMesh(InMergeMesh)
 ,	SrcMeshList(InSrcMeshList)
 ,	StripTopLODs(InStripTopLODs)
@@ -39,6 +39,40 @@ FSkeletalMeshMerge::FSkeletalMeshMerge(USkeletalMesh* InMergeMesh,
 	check(MergeMesh);
 }
 
+
+FSkeletalMeshMerge::FSkeletalMeshMerge(USkeletalMesh* InMergeMesh, 
+										const TArray<USkeletalMesh*>& InSrcMeshList, 
+										const TArray<FSkelMeshMergeSectionMapping>& InForceSectionMapping,
+										int32 InStripTopLODs,
+										EMeshBufferAccess InMeshBufferAccess,
+										PRAGMA_DISABLE_DEPRECATION_WARNINGS
+										FSkelMeshMergeUVTransforms* InSectionUVTransforms)
+										PRAGMA_ENABLE_DEPRECATION_WARNINGS
+:	MergeMesh(InMergeMesh)
+,	SrcMeshList(InSrcMeshList)
+,	StripTopLODs(InStripTopLODs)
+,   MeshBufferAccess(InMeshBufferAccess)
+,	ForceSectionMapping(InForceSectionMapping)
+PRAGMA_DISABLE_DEPRECATION_WARNINGS
+,	SectionUVTransforms(&DummySectionUVTransforms)
+PRAGMA_ENABLE_DEPRECATION_WARNINGS
+{
+	if (InSectionUVTransforms)
+	{
+		PRAGMA_DISABLE_DEPRECATION_WARNINGS
+		for (int32 MeshIdx = 0; MeshIdx < InSrcMeshList.Num(); ++MeshIdx)
+		{
+			if (MeshIdx < InSectionUVTransforms->UVTransformsPerMesh.Num())
+			{
+				const TArray<FTransform>& OldPerMeshTransforms = InSectionUVTransforms->UVTransformsPerMesh[MeshIdx];
+				FSkelMeshMergeMeshUVTransforms& NewPerMeshTransforms = DummySectionUVTransforms.UVTransformsPerMesh.AddDefaulted_GetRef();
+				NewPerMeshTransforms.UVTransforms = OldPerMeshTransforms;
+			}
+		}
+		PRAGMA_ENABLE_DEPRECATION_WARNINGS
+	}
+}
+
 /** Helper macro to call GenerateLODModel which requires compile time vertex type. */
 #define GENERATE_LOD_MODEL( VertexType, NumUVs ) \
 {\
@@ -330,7 +364,7 @@ void FSkeletalMeshMerge::GenerateNewSectionArray( TArray<FNewSectionInfo>& NewSe
 							TArray<FTransform> SrcUVTransform;
 							if (SectionUVTransforms != nullptr && MeshIdx < SectionUVTransforms->UVTransformsPerMesh.Num())
 							{
-								SrcUVTransform = SectionUVTransforms->UVTransformsPerMesh[MeshIdx];
+								SrcUVTransform = SectionUVTransforms->UVTransformsPerMesh[MeshIdx].UVTransforms;
 							}
 
 							// add the source section as a new merge entry
@@ -367,7 +401,7 @@ void FSkeletalMeshMerge::GenerateNewSectionArray( TArray<FNewSectionInfo>& NewSe
 					TArray<FTransform> SrcUVTransform;
 					if (SectionUVTransforms != nullptr && MeshIdx < SectionUVTransforms->UVTransformsPerMesh.Num())
 					{
-						SrcUVTransform = SectionUVTransforms->UVTransformsPerMesh[MeshIdx];
+						SrcUVTransform = SectionUVTransforms->UVTransformsPerMesh[MeshIdx].UVTransforms;
 					}
 					// add a new merge section entry
 					FMergeSectionInfo& MergeSectionInfo = *new(NewSectionInfo.MergeSections) FMergeSectionInfo(
diff --git a/Engine/Source/Runtime/Engine/Public/ReferenceSkeleton.h b/Engine/Source/Runtime/Engine/Public/ReferenceSkeleton.h
--- a/Engine/Source/Runtime/Engine/Public/ReferenceSkeleton.h
+++ b/Engine/Source/Runtime/Engine/Public/ReferenceSkeleton.h
@@ -76,6 +76,7 @@ private:
 	const USkeleton*	Skeleton;
 public:
 	FReferenceSkeletonModifier(FReferenceSkeleton& InRefSkel, const USkeleton* InSkeleton) : RefSkeleton(InRefSkel), Skeleton(InSkeleton) {}
+	FReferenceSkeletonModifier(USkeleton* InSkeleton);
 	~FReferenceSkeletonModifier();
 
 	// Update the reference pose transform of the specified bone
diff --git a/Engine/Source/Runtime/Engine/Public/SkeletalMeshMerge.h b/Engine/Source/Runtime/Engine/Public/SkeletalMeshMerge.h
--- a/Engine/Source/Runtime/Engine/Public/SkeletalMeshMerge.h
+++ b/Engine/Source/Runtime/Engine/Public/SkeletalMeshMerge.h
@@ -15,6 +15,8 @@ FSkeletalMeshMerge
 #include "ReferenceSkeleton.h"
 #include "Components.h"
 
+#include "SkeletalMeshMerge.generated.h"
+
 class UMaterialInterface;
 class USkeletalMesh;
 class USkeletalMeshSocket;
@@ -78,21 +80,48 @@ struct FRefPoseOverride
 * Info to map all the sections from a single source skeletal mesh to 
 * a final section entry int he merged skeletal mesh
 */
-struct FSkelMeshMergeSectionMapping
+USTRUCT(BlueprintType)
+struct ENGINE_API FSkelMeshMergeSectionMapping
 {
+	GENERATED_USTRUCT_BODY()
+	
 	/** indices to final section entries of the merged skel mesh */
+	UPROPERTY(EditAnywhere, BlueprintReadWrite, Category = "Mesh Merge Parameters")
 	TArray<int32> SectionIDs;
 };
 
+USTRUCT(BlueprintType)
+struct ENGINE_API FSkelMeshMergeMeshUVTransforms
+{
+	GENERATED_USTRUCT_BODY()
+
+	/** A list of how UVs should be transformed on a given mesh, where index represents a specific UV channel. */
+	UPROPERTY(EditAnywhere, BlueprintReadWrite, Category = "Mesh Merge Parameters")
+	TArray<FTransform> UVTransforms;
+};
+
 /** 
 * Info to map all the sections about how to transform their UVs
 */
-struct FSkelMeshMergeUVTransforms
+struct UE_DEPRECATED(5.0, "FSkelMeshMergeUVTransforms has been deprecated, use FSkelMeshMergeMeshUVTransforms instead") FSkelMeshMergeUVTransforms
 {
 	/** For each UV channel on each mesh, how the UVS should be transformed. */
 	TArray<TArray<FTransform>> UVTransformsPerMesh;
 };
 
+/** 
+* Info to map all the sections about how to transform their UVs
+*/
+USTRUCT(BlueprintType)
+struct ENGINE_API FSkelMeshMergeUVTransformMapping
+{
+	GENERATED_USTRUCT_BODY()
+	
+	/** UV coordinates transform datam one entry for each Skeletal Mesh. */
+	UPROPERTY(EditAnywhere, BlueprintReadWrite, Category = "Mesh Merge Parameters")
+	TArray<FSkelMeshMergeMeshUVTransforms> UVTransformsPerMesh;
+};
+
 /** 
 * Utility for merging a list of skeletal meshes into a single mesh.
 */
@@ -114,9 +143,21 @@ public:
 		const TArray<FSkelMeshMergeSectionMapping>& InForceSectionMapping,
 		int32 StripTopLODs,
         EMeshBufferAccess MeshBufferAccess=EMeshBufferAccess::Default,
-		FSkelMeshMergeUVTransforms* InSectionUVTransforms = nullptr
+		const FSkelMeshMergeUVTransformMapping* InSectionUVTransforms = nullptr
 		);
 
+	UE_DEPRECATED(5.0, "FSkelMeshMergeUVTransforms has been replaced with FSkelMeshMergeMeshUVTransforms, use different signature")
+	FSkeletalMeshMerge( 
+    		USkeletalMesh* InMergeMesh, 
+    		const TArray<USkeletalMesh*>& InSrcMeshList, 
+    		const TArray<FSkelMeshMergeSectionMapping>& InForceSectionMapping,
+    		int32 StripTopLODs,
+            EMeshBufferAccess MeshBufferAccess,
+			PRAGMA_DISABLE_DEPRECATION_WARNINGS
+    		FSkelMeshMergeUVTransforms* InSectionUVTransforms
+    		PRAGMA_ENABLE_DEPRECATION_WARNINGS
+    		);
+
 	/**
 	 * Merge/Composite skeleton and meshes together from the list of source meshes.
 	 * @param RefPoseOverrides - An optional override for the merged skeleton's reference pose.
@@ -168,7 +209,10 @@ private:
 	const TArray<FSkelMeshMergeSectionMapping>& ForceSectionMapping;
 
 	/** optional array to transform UVs in each source mesh */
-	const FSkelMeshMergeUVTransforms* SectionUVTransforms;
+	const FSkelMeshMergeUVTransformMapping* SectionUVTransforms;
+
+	UE_DEPRECATED(5.0, "Used to facilitate backwards compatibility with old constructor")
+	FSkelMeshMergeUVTransformMapping DummySectionUVTransforms;
 
 	/** Matches the Materials array in the final mesh - used for creating the right number of Material slots. */
 	TArray<int32>	MaterialIds;
