Fix non-unity compile issues w/ SkeletalMerging

diff --git a/Engine/Plugins/Runtime/SkeletalMerging/Source/SkeletalMerging/Private/SkeletalMergingLibrary.cpp b/Engine/Plugins/Runtime/SkeletalMerging/Source/SkeletalMerging/Private/SkeletalMergingLibrary.cpp
--- a/Engine/Plugins/Runtime/SkeletalMerging/Source/SkeletalMerging/Private/SkeletalMergingLibrary.cpp
+++ b/Engine/Plugins/Runtime/SkeletalMerging/Source/SkeletalMerging/Private/SkeletalMergingLibrary.cpp
@@ -6,6 +6,8 @@
 #include "Engine/SkeletalMesh.h"
 #include "Animation/Skeleton.h"
 #include "Algo/Accumulate.h"
+#include "Algo/Transform.h"
+#include "Animation/BlendProfile.h"
 
 IMPLEMENT_MODULE(IModuleInterface, SkeletalMerging);
 
diff --git a/Engine/Plugins/Runtime/SkeletalMerging/Source/SkeletalMerging/Public/SkeletalMergingLibrary.h b/Engine/Plugins/Runtime/SkeletalMerging/Source/SkeletalMerging/Public/SkeletalMergingLibrary.h
--- a/Engine/Plugins/Runtime/SkeletalMerging/Source/SkeletalMerging/Public/SkeletalMergingLibrary.h
+++ b/Engine/Plugins/Runtime/SkeletalMerging/Source/SkeletalMerging/Public/SkeletalMergingLibrary.h
@@ -3,6 +3,7 @@
 #pragma once
 
 #include "SkeletalMeshMerge.h"
+#include "Kismet/BlueprintFunctionLibrary.h"
 #include "SkeletalMergingLibrary.generated.h"
 
 /**
