Final non-unity error fix in SkeletalMergingLibrary

diff --git a/Engine/Plugins/Runtime/SkeletalMerging/Source/SkeletalMerging/Private/SkeletalMergingLibrary.cpp b/Engine/Plugins/Runtime/SkeletalMerging/Source/SkeletalMerging/Private/SkeletalMergingLibrary.cpp
--- a/Engine/Plugins/Runtime/SkeletalMerging/Source/SkeletalMerging/Private/SkeletalMergingLibrary.cpp
+++ b/Engine/Plugins/Runtime/SkeletalMerging/Source/SkeletalMerging/Private/SkeletalMergingLibrary.cpp
@@ -8,6 +8,8 @@
 #include "Algo/Accumulate.h"
 #include "Algo/Transform.h"
 #include "Animation/BlendProfile.h"
+#include "Modules/ModuleManager.h"
+#include "Modules/ModuleInterface.h"
 
 IMPLEMENT_MODULE(IModuleInterface, SkeletalMerging);
 
