//UE5/Main - Compile QAGame Mac - Loop variables in 'SkeletalMergingLibrary.cpp' and 'MeshDescription.cpp' are creating copies
#fix fixed unnecessary copy in SkeletalMergingLibrary

diff --git a/Engine/Plugins/Runtime/SkeletalMerging/Source/SkeletalMerging/Private/SkeletalMergingLibrary.cpp b/Engine/Plugins/Runtime/SkeletalMerging/Source/SkeletalMerging/Private/SkeletalMergingLibrary.cpp
--- a/Engine/Plugins/Runtime/SkeletalMerging/Source/SkeletalMerging/Private/SkeletalMergingLibrary.cpp
+++ b/Engine/Plugins/Runtime/SkeletalMerging/Source/SkeletalMerging/Private/SkeletalMergingLibrary.cpp
@@ -292,7 +292,7 @@ void USkeletalMergingLibrary::AddCurveNames(USkeleton* InSkeleton, const TMap<FN
 	});
 	InSkeleton->VerifySmartNames(USkeleton::AnimCurveMappingName, CurveSmartNames);
 		
-	for(const TPair<FName, const FCurveMetaData*> CurveMetaDataPair : InCurves)
+	for(const TPair<FName, const FCurveMetaData*>& CurveMetaDataPair : InCurves)
 	{
 		if (CurveMetaDataPair.Value)
 		{
