Skeletal Mesh merging plugin:
+ Added optional check to detect incompatible skeletons when merging skeletons

diff --git a/Engine/Plugins/Runtime/SkeletalMerging/Source/SkeletalMerging/Private/SkeletalMergingLibrary.cpp b/Engine/Plugins/Runtime/SkeletalMerging/Source/SkeletalMerging/Private/SkeletalMergingLibrary.cpp
--- a/Engine/Plugins/Runtime/SkeletalMerging/Source/SkeletalMerging/Private/SkeletalMergingLibrary.cpp
+++ b/Engine/Plugins/Runtime/SkeletalMerging/Source/SkeletalMerging/Private/SkeletalMergingLibrary.cpp
@@ -13,6 +13,8 @@
 
 IMPLEMENT_MODULE(IModuleInterface, SkeletalMerging);
 
+DEFINE_LOG_CATEGORY(LogSkeletalMeshMerge);
+
 namespace UE
 {
 	namespace SkeletonMerging
@@ -121,6 +123,10 @@ USkeleton* USkeletalMergingLibrary::MergeSkeletons(const FSkeletonMergeParams& P
 	TMap<FName, uint32> BoneNamesToPathHash;
 	BoneNamesToPathHash.Reserve(TotalPossibleBones);
 
+	// Bone name to bone pose 
+	TMap<FName, FTransform> BoneNamesToBonePose;
+	BoneNamesToBonePose.Reserve(TotalPossibleBones);
+
 	// Combined bone and socket name hash
 	TMap<uint32, TObjectPtr<USkeletalMeshSocket>> HashToSockets;
 	// Combined from and to-bone name hash
@@ -130,6 +136,8 @@ USkeleton* USkeletalMergingLibrary::MergeSkeletons(const FSkeletonMergeParams& P
 	TMap<FName, TSet<FName>> GroupToSlotNames;
 	TMap<FName, TArray<const UBlendProfile*>> UniqueBlendProfiles;
 
+	bool bMergeSkeletonsFailed = false;
+
 	for (int32 SkeletonIndex = 0; SkeletonIndex < NumberOfSkeletons; ++SkeletonIndex)
 	{
 		const USkeleton* Skeleton = ToMergeSkeletons[SkeletonIndex];
@@ -137,6 +145,9 @@ USkeleton* USkeletalMergingLibrary::MergeSkeletons(const FSkeletonMergeParams& P
 		const TArray<FMeshBoneInfo>& Bones = ReferenceSkeleton.GetRawRefBoneInfo();		
 		const TArray<FTransform>& BonePoses = ReferenceSkeleton.GetRawRefBonePose();
 
+
+		bool bConflictivePoseFound = false;
+
 		const int32 NumBones = Bones.Num();
 		for (int32 BoneIndex = 0; BoneIndex < NumBones; ++BoneIndex)
 		{
@@ -152,10 +163,44 @@ USkeleton* USkeletalMergingLibrary::MergeSkeletons(const FSkeletonMergeParams& P
 
 			// Append parent hash to path to give full path hash to current bone
 			const uint32 BonePathHash = HashCombine(ParentPathHash, ParentHash);
+
+			if (Params.bCheckSkeletonsCompatibility)
+			{
+				// Check if the bone exists in the hierarchy 
+				if (const uint32* ExistingPath = BoneNamesToPathHash.Find(Bone.Name))
+				{
+					const uint32 ExistingPathHash = *ExistingPath;
+
+					// If the hash differs from the existing one it means skeletons are incompatible
+					if (ExistingPathHash != BonePathHash)
+					{
+						UE_LOG(LogSkeletalMeshMerge, Error, TEXT("Failed to merge skeletons. Skeleton %s has an invalid bone chain."), *Skeleton->GetName());
+						bMergeSkeletonsFailed = true;
+						break;
+					}
+
+					// Bone poses will be overwritten, check if they are the same
+					if (!bConflictivePoseFound && !BoneNamesToBonePose[Bone.Name].Equals(BonePoses[BoneIndex]))
+					{
+						UE_LOG(LogSkeletalMeshMerge, Warning, TEXT("Skeleton %s has a different reference pose, reference pose will be overwritten."), *Skeleton->GetName());
+						bConflictivePoseFound = true;
+					}
+				}
+
+				BoneNamesToBonePose.Add(Bone.Name, BonePoses[BoneIndex]);
+			}
+			
+			// Add path hash to current bone
 			BoneNamesToPathHash.Add(Bone.Name, BonePathHash);
 
 			// Add bone to hierarchy
 			MergedBoneHierarchy.AddBone(Bone.Name, BonePoses[BoneIndex], BonePathHash);
+			
+		}
+
+		if (Params.bCheckSkeletonsCompatibility && bMergeSkeletonsFailed)
+		{
+			continue;
 		}
 
 		if (Params.bMergeSockets)
@@ -208,6 +253,12 @@ USkeleton* USkeletalMergingLibrary::MergeSkeletons(const FSkeletonMergeParams& P
 		}		
 	}
 
+	if (bMergeSkeletonsFailed)
+	{
+		UE_LOG(LogSkeletalMeshMerge, Error, TEXT("Failed to merge skeletons. One or more skeletons with invalid parent chains were found."));
+		return nullptr;
+	}
+
 	USkeleton* GeneratedSkeleton = NewObject<USkeleton>();
 
 	// Generate bone hierarchy
diff --git a/Engine/Plugins/Runtime/SkeletalMerging/Source/SkeletalMerging/Public/SkeletalMergingLibrary.h b/Engine/Plugins/Runtime/SkeletalMerging/Source/SkeletalMerging/Public/SkeletalMergingLibrary.h
--- a/Engine/Plugins/Runtime/SkeletalMerging/Source/SkeletalMerging/Public/SkeletalMergingLibrary.h
+++ b/Engine/Plugins/Runtime/SkeletalMerging/Source/SkeletalMerging/Public/SkeletalMergingLibrary.h
@@ -6,6 +6,8 @@
 #include "Kismet/BlueprintFunctionLibrary.h"
 #include "SkeletalMergingLibrary.generated.h"
 
+DECLARE_LOG_CATEGORY_EXTERN(LogSkeletalMeshMerge, Log, All);
+
 /**
 * Struct containing all parameters used to perform a Skeletal Mesh merge.
 */
@@ -87,6 +89,10 @@ struct SKELETALMERGING_API FSkeletonMergeParams
 	// Whether or not to include Animation Slot Group (names) when merging the Skeletons
 	UPROPERTY(EditAnywhere, BlueprintReadWrite, Category=SkeletonMerge)
 	bool bMergeAnimSlotGroups = true;
+
+	// Whether or not to check if there are invalid parent chains or shared bones with different reference transforms
+	UPROPERTY(EditAnywhere, BlueprintReadWrite, Category=SkeletonMerge)
+	bool bCheckSkeletonsCompatibility = false;
 };
 
 /**
