Add functionality to add named virtual bones to a Skeleton. Fix SkeletalMergingLibrary virtual bone merging names.

- A new Skeleton member function AddNewNamedVirtualBone has been added.
- Virtual bone names would change to the default Source_Target naming when participating of a skeleton merge, now the original name is used.

#rnx

diff --git a/Engine/Plugins/Runtime/SkeletalMerging/Source/SkeletalMerging/Private/SkeletalMergingLibrary.cpp b/Engine/Plugins/Runtime/SkeletalMerging/Source/SkeletalMerging/Private/SkeletalMergingLibrary.cpp
--- a/Engine/Plugins/Runtime/SkeletalMerging/Source/SkeletalMerging/Private/SkeletalMergingLibrary.cpp
+++ b/Engine/Plugins/Runtime/SkeletalMerging/Source/SkeletalMerging/Private/SkeletalMergingLibrary.cpp
@@ -273,7 +273,7 @@ USkeleton* USkeletalMergingLibrary::MergeSkeletons(const FSkeletonMergeParams& P
 	{
 		TArray<const FVirtualBone*> VirtualBones;
 		HashToVirtualBones.GenerateValueArray(VirtualBones);
-		AddVirtualBones(GeneratedSkeleton,VirtualBones);
+		AddVirtualBones(GeneratedSkeleton, VirtualBones);
 	}
 
 	// Merge Curve / track mappings	
@@ -317,12 +317,11 @@ void USkeletalMergingLibrary::AddSockets(USkeleton* InSkeleton, const TArray<TOb
 	}
 }
 
-void USkeletalMergingLibrary::AddVirtualBones(USkeleton* InSkeleton, const TArray<const FVirtualBone*> InVirtualBones)
+void USkeletalMergingLibrary::AddVirtualBones(USkeleton* InSkeleton, const TArray<const FVirtualBone*>& InVirtualBones)
 {
-	for(const FVirtualBone* VirtualBone : InVirtualBones)
+	for (const FVirtualBone* VirtualBone : InVirtualBones)
 	{
-		FName VirtualBoneName = NAME_None;				
-		InSkeleton->AddNewVirtualBone(VirtualBone->SourceBoneName, VirtualBone->TargetBoneName, VirtualBoneName);
+		InSkeleton->AddNewNamedVirtualBone(VirtualBone->SourceBoneName, VirtualBone->TargetBoneName, VirtualBone->VirtualBoneName);
 	}	
 }
 
diff --git a/Engine/Plugins/Runtime/SkeletalMerging/Source/SkeletalMerging/Public/SkeletalMergingLibrary.h b/Engine/Plugins/Runtime/SkeletalMerging/Source/SkeletalMerging/Public/SkeletalMergingLibrary.h
--- a/Engine/Plugins/Runtime/SkeletalMerging/Source/SkeletalMerging/Public/SkeletalMergingLibrary.h
+++ b/Engine/Plugins/Runtime/SkeletalMerging/Source/SkeletalMerging/Public/SkeletalMergingLibrary.h
@@ -119,7 +119,7 @@ public:
 
 protected:
 	static void AddSockets(USkeleton* InSkeleton, const TArray<TObjectPtr<class USkeletalMeshSocket>>& InSockets);
-	static void AddVirtualBones(USkeleton* InSkeleton, const TArray<const struct FVirtualBone*> InVirtualBones);
+	static void AddVirtualBones(USkeleton* InSkeleton, const TArray<const struct FVirtualBone*>& InVirtualBones);
 	static void AddCurveNames(USkeleton* InSkeleton, const TMap<FName, const struct FCurveMetaData*>& InCurves);
 	static void AddBlendProfiles(USkeleton* InSkeleton, const TMap<FName, TArray<const class UBlendProfile*>>& InBlendProfiles);
 	static void AddAnimationSlotGroups(USkeleton* InSkeleton, const TMap<FName, TSet<FName>>& InSlotGroupsNames);
diff --git a/Engine/Source/Runtime/Engine/Classes/Animation/Skeleton.h b/Engine/Source/Runtime/Engine/Classes/Animation/Skeleton.h
--- a/Engine/Source/Runtime/Engine/Classes/Animation/Skeleton.h
+++ b/Engine/Source/Runtime/Engine/Classes/Animation/Skeleton.h
@@ -246,6 +246,7 @@ namespace VirtualBoneNameHelpers
 
 	ENGINE_API FString AddVirtualBonePrefix(const FString& InName);
 	ENGINE_API FName RemoveVirtualBonePrefix(const FString& InName);
+	ENGINE_API bool CheckVirtualBonePrefix(const FString& InName);
 }
 
 USTRUCT()
@@ -265,6 +266,13 @@ public:
 
 	FVirtualBone() {}
 
+	FVirtualBone(FName InSource, FName InTarget, FName InVirtual)
+		: SourceBoneName(InSource)
+		, TargetBoneName(InTarget)
+		, VirtualBoneName(InVirtual)
+	{
+	}
+
 	FVirtualBone(FName InSource, FName InTarget)
 		: SourceBoneName(InSource)
 		, TargetBoneName(InTarget)
@@ -434,6 +442,8 @@ public:
 
 	ENGINE_API bool AddNewVirtualBone(const FName SourceBoneName, const FName TargetBoneName, FName& NewVirtualBoneName);
 
+	ENGINE_API bool AddNewNamedVirtualBone(const FName SourceBoneName, const FName TargetBoneName, const FName NewVirtualBoneName);
+	
 	ENGINE_API void RemoveVirtualBones(const TArray<FName>& BonesToRemove);
 
 	ENGINE_API void RenameVirtualBone(const FName OriginalBoneName, const FName NewBoneName);
diff --git a/Engine/Source/Runtime/Engine/Private/Animation/Skeleton.cpp b/Engine/Source/Runtime/Engine/Private/Animation/Skeleton.cpp
--- a/Engine/Source/Runtime/Engine/Private/Animation/Skeleton.cpp
+++ b/Engine/Source/Runtime/Engine/Private/Animation/Skeleton.cpp
@@ -120,6 +120,11 @@ namespace VirtualBoneNameHelpers
 	{
 		return FName(SkipPrefix(InName));
 	}
+
+	bool CheckVirtualBonePrefix(const FString& InName)
+	{
+		return InName.StartsWith(VirtualBonePrefix);
+	}
 }
 
 #if WITH_EDITORONLY_DATA
@@ -1933,6 +1938,32 @@ bool USkeleton::AddNewVirtualBone(const FName SourceBoneName, const FName Target
 	return true;
 }
 
+bool USkeleton::AddNewNamedVirtualBone(const FName SourceBoneName, const FName TargetBoneName, const FName VirtualBoneName)
+{
+	if (!VirtualBoneNameHelpers::CheckVirtualBonePrefix(VirtualBoneName.ToString()))
+	{
+		return false;
+	}
+
+	for (const FVirtualBone& SSBone : VirtualBones)
+	{
+		if ((SSBone.SourceBoneName == SourceBoneName && SSBone.TargetBoneName == TargetBoneName) ||
+			SSBone.VirtualBoneName == VirtualBoneName)
+		{
+			return false;
+		}
+	}
+
+	Modify();
+
+	VirtualBones.Emplace(SourceBoneName, TargetBoneName, VirtualBoneName);
+
+	RegenerateVirtualBoneGuid();
+	HandleVirtualBoneChanges();
+
+	return true;
+}
+
 int32 FindBoneByName(const FName& BoneName, TArray<FVirtualBone>& Bones)
 {
 	for (int32 Idx = 0; Idx < Bones.Num(); ++Idx)
